Motion Design: added data link oauth plugin

diff --git a/Engine/Plugins/Experimental/DataLinkOAuth/DataLinkOAuth.uplugin b/Engine/Plugins/Experimental/DataLinkOAuth/DataLinkOAuth.uplugin
new file mode 100644
--- /dev/null
+++ b/Engine/Plugins/Experimental/DataLinkOAuth/DataLinkOAuth.uplugin
@@ -0,0 +1,34 @@
+{
+	"FileVersion": 3,
+	"Version": 1,
+	"VersionName": "1.0",
+	"FriendlyName": "DataLink OAuth",
+	"Description": "Data Link functionality for OAuth 2.0",
+	"Category": "Other",
+	"CreatedBy": "",
+	"CreatedByURL": "",
+	"DocsURL": "",
+	"MarketplaceURL": "",
+	"CanContainContent": true,
+	"IsBetaVersion": false,
+	"IsExperimentalVersion": true,
+	"Installed": false,
+	"Modules": [
+		{
+			"Name": "DataLinkOAuth",
+			"Type": "Runtime",
+			"LoadingPhase": "Default"
+		},
+		{
+			"Name": "DataLinkOAuthEditor",
+			"Type": "Editor",
+			"LoadingPhase": "Default"
+		}
+	],
+	"Plugins": [
+		{
+			"Name": "DataLink",
+			"Enabled": true
+		}
+	]
+}
\ No newline at end of file
diff --git a/Engine/Plugins/Experimental/DataLinkOAuth/Source/DataLinkOAuth/DataLinkOAuth.Build.cs b/Engine/Plugins/Experimental/DataLinkOAuth/Source/DataLinkOAuth/DataLinkOAuth.Build.cs
new file mode 100644
--- /dev/null
+++ b/Engine/Plugins/Experimental/DataLinkOAuth/Source/DataLinkOAuth/DataLinkOAuth.Build.cs
@@ -0,0 +1,32 @@
+// Copyright Epic Games, Inc. All Rights Reserved.
+
+using UnrealBuildTool;
+
+public class DataLinkOAuth : ModuleRules
+{
+	public DataLinkOAuth(ReadOnlyTargetRules Target) : base(Target)
+	{
+		PCHUsage = ModuleRules.PCHUsageMode.UseExplicitOrSharedPCHs;
+
+		PublicDependencyModuleNames.AddRange(
+			new string[]
+			{
+				"Core",
+				"CoreUObject",
+				"DataLink",
+				"DataLinkHttp",
+			}
+		);
+
+		PrivateDependencyModuleNames.AddRange(
+			new string[]
+			{
+				"Engine",
+				"HTTP",
+				"HTTPServer",
+				"Json",
+				"Sockets",
+			}
+		);
+	}
+}
diff --git a/Engine/Plugins/Experimental/DataLinkOAuth/Source/DataLinkOAuth/Private/DataLinkNodeOAuth.cpp b/Engine/Plugins/Experimental/DataLinkOAuth/Source/DataLinkOAuth/Private/DataLinkNodeOAuth.cpp
new file mode 100644
--- /dev/null
+++ b/Engine/Plugins/Experimental/DataLinkOAuth/Source/DataLinkOAuth/Private/DataLinkNodeOAuth.cpp
@@ -0,0 +1,206 @@
+// Copyright Epic Games, Inc. All Rights Reserved.
+
+#include "DataLinkNodeOAuth.h"
+#include "DataLinkCoreTypes.h"
+#include "DataLinkExecutor.h"
+#include "DataLinkHttpSettings.h"
+#include "DataLinkNames.h"
+#include "DataLinkOAuthAuthRequest.h"
+#include "DataLinkOAuthCodeAccessExchange.h"
+#include "DataLinkOAuthHandle.h"
+#include "DataLinkOAuthInstance.h"
+#include "DataLinkOAuthLog.h"
+#include "DataLinkOAuthSettings.h"
+#include "DataLinkOAuthSubsystem.h"
+#include "DataLinkOAuthToken.h"
+#include "DataLinkPinBuilder.h"
+#include "HttpServerRequest.h"
+
+#define LOCTEXT_NAMESPACE "DataLinkNodeOAuth"
+
+UDataLinkNodeOAuth::UDataLinkNodeOAuth()
+{
+	InstanceStruct = FDataLinkNodeOAuthInstance::StaticStruct();
+}
+
+void UDataLinkNodeOAuth::OnBuildPins(FDataLinkPinBuilder& Inputs, FDataLinkPinBuilder& Outputs) const
+{
+	Inputs.Add(UE::DataLinkOAuth::InputHttp)
+		.SetDisplayName(LOCTEXT("HttpSettingsDisplay", "Http Settings"))
+		.SetStruct<FDataLinkHttpSettings>();
+
+	Inputs.Add(UE::DataLinkOAuth::InputOAuth)
+		.SetDisplayName(LOCTEXT("OAuthSettingsDisplay", "OAuth Settings"))
+		.SetStruct<FDataLinkOAuthSettingsWrapper>();
+
+	Outputs.Add(UE::DataLink::OutputDefault)
+		.SetDisplayName(LOCTEXT("OutputDisplay", "Http Settings"))
+		.SetStruct<FDataLinkHttpSettings>();
+}
+
+EDataLinkExecutionReply UDataLinkNodeOAuth::OnExecute(FDataLinkExecutor& InExecutor) const
+{
+	FDataLinkNodeInstance& NodeInstance = InExecutor.GetNodeInstanceMutable(this);
+
+	const FDataLinkInputDataViewer& InputDataViewer = NodeInstance.GetInputDataViewer();
+	const FDataLinkOutputDataViewer& OutputDataViewer = NodeInstance.GetOutputDataViewer();
+
+	UDataLinkOAuthSettings* const OAuthSettings = InputDataViewer.Get<FDataLinkOAuthSettingsWrapper>(UE::DataLinkOAuth::InputOAuth).OAuthSettings;
+	if (!OAuthSettings)
+	{
+		UE_LOG(LogDataLinkOAuth, Error, TEXT("[%s] Data Link OAuth - Invalid OAuth settings."), InExecutor.GetContextName().GetData());
+		return EDataLinkExecutionReply::Unhandled;
+	}
+
+	// Attempt to re-use an existing access token if available and not expired.
+	if (UDataLinkOAuthSubsystem* OAuthSubsystem = UDataLinkOAuthSubsystem::Get())
+	{
+		OAuthSubsystem->CleanExpiredTokens();
+
+		if (const FDataLinkOAuthToken* OAuthToken = OAuthSubsystem->FindToken(OAuthSettings))
+		{
+			// Copy Input Http Settings into Output
+			FDataLinkHttpSettings& OutputHttpSettings = OutputDataViewer.Get<FDataLinkHttpSettings>(UE::DataLink::OutputDefault);
+			OutputHttpSettings = InputDataViewer.Get<FDataLinkHttpSettings>(UE::DataLinkOAuth::InputHttp);
+
+			if (OAuthSettings->AuthorizeHttpRequest(*OAuthToken, OutputHttpSettings))
+			{
+				InExecutor.SucceedNode(this, FConstStructView::Make(OutputHttpSettings));
+				return EDataLinkExecutionReply::Handled;
+			}
+		}
+	}
+
+	FDataLinkNodeOAuthInstance& OAuthInstance = NodeInstance.GetInstanceDataMutable().Get<FDataLinkNodeOAuthInstance>();
+	OAuthInstance.SharedData = OAuthSettings->MakeSharedData();
+
+	UE::DataLinkOAuth::FAuthRequestParams AuthRequestParams;
+	AuthRequestParams.OAuthSettings = OAuthSettings;
+	AuthRequestParams.OAuthInstanceView = OAuthInstance;
+	AuthRequestParams.OnAuthResponse.BindUObject(this, &UDataLinkNodeOAuth::OnAuthResponse, InExecutor.AsWeak());
+
+	if (!UE::DataLinkOAuth::RequestAuthorization(AuthRequestParams))
+	{
+		UE_LOG(LogDataLinkOAuth, Error, TEXT("[%s] Data Link OAuth - Authorization request failed."), InExecutor.GetContextName().GetData());
+		return EDataLinkExecutionReply::Unhandled;
+	}
+
+	return EDataLinkExecutionReply::Handled;
+}
+
+void UDataLinkNodeOAuth::OnAuthResponse(const FHttpServerRequest& InRequest, TWeakPtr<FDataLinkExecutor> InExecutorWeak) const
+{
+	TSharedPtr<FDataLinkExecutor> Executor = InExecutorWeak.Pin();
+	if (!Executor.IsValid())
+	{
+		return;
+	}
+
+	FDataLinkNodeInstance& NodeInstance = Executor->GetNodeInstanceMutable(this);
+
+	const FDataLinkInputDataViewer& InputDataViewer = NodeInstance.GetInputDataViewer();
+
+	UDataLinkOAuthSettings* const OAuthSettings = InputDataViewer.Get<FDataLinkOAuthSettingsWrapper>(UE::DataLinkOAuth::InputOAuth).OAuthSettings;
+	check(OAuthSettings);
+
+	FDataLinkNodeOAuthInstance& OAuthInstance = NodeInstance.GetInstanceDataMutable().Get<FDataLinkNodeOAuthInstance>();
+	if (!OAuthSettings->ValidateRequest(InRequest, OAuthInstance))
+	{
+		// Request did not come from this OAuth instance. Skip
+		// NOTE: FailNode isn't called here because the expected request might still come
+		return;
+	}
+
+	// Stop listening as the expected request was found
+	OAuthInstance.StopListening();
+
+	FStringView AuthCodeView;
+	if (!OAuthSettings->FindAuthCode(InRequest, OAuthInstance, AuthCodeView) || AuthCodeView.IsEmpty())
+	{
+		UE_LOG(LogDataLinkOAuth, Error, TEXT("[%s] Data Link OAuth - Invalid Authorization Code (Settings: %s)")
+			, Executor->GetContextName().GetData()
+			, *OAuthSettings->GetName());
+
+		Executor->FailNode(this);
+		return;
+	}
+
+	UE::DataLinkOAuth::FExchangeAuthCodeParams ExchangeAuthCodeParams;
+	ExchangeAuthCodeParams.OAuthSettings = OAuthSettings;
+	ExchangeAuthCodeParams.AuthCodeView = AuthCodeView;
+	ExchangeAuthCodeParams.OAuthInstanceView = OAuthInstance;
+	ExchangeAuthCodeParams.OnResponse.BindUObject(this, &UDataLinkNodeOAuth::OnExchangeCodeResponse, InExecutorWeak);
+
+	if (!UE::DataLinkOAuth::ExchangeAuthCodeForAccess(ExchangeAuthCodeParams))
+	{
+		UE_LOG(LogDataLinkOAuth, Error, TEXT("[%s] Data Link OAuth - Failed to Exchange Auth Code For Access (Settings: %s)")
+			, Executor->GetContextName().GetData()
+			, *OAuthSettings->GetName());
+
+		Executor->FailNode(this);
+	}
+}
+
+void UDataLinkNodeOAuth::OnExchangeCodeResponse(FStringView InResponse, TWeakPtr<FDataLinkExecutor> InExecutorWeak) const
+{
+	TSharedPtr<FDataLinkExecutor> Executor = InExecutorWeak.Pin();
+	if (!Executor.IsValid())
+	{
+		return;
+	}
+
+	const FDataLinkNodeInstance& NodeInstance = Executor->GetNodeInstance(this);
+
+	const FDataLinkInputDataViewer& InputDataViewer = NodeInstance.GetInputDataViewer();
+
+	UDataLinkOAuthSettings* const OAuthSettings = InputDataViewer.Get<FDataLinkOAuthSettingsWrapper>(UE::DataLinkOAuth::InputOAuth).OAuthSettings;
+	check(OAuthSettings);
+
+	if (InResponse.IsEmpty())
+	{
+		UE_LOG(LogDataLinkOAuth, Error, TEXT("[%s] Data Link OAuth - Failed to get a valid response with Access Token (Settings: %s)")
+			, Executor->GetContextName().GetData()
+			, *OAuthSettings->GetName());
+
+		Executor->FailNode(this);
+		return;
+	}
+
+	FDataLinkOAuthToken OAuthToken;
+	if (!OAuthSettings->BuildAuthToken(InResponse, OAuthToken))
+	{
+		UE_LOG(LogDataLinkOAuth, Error, TEXT("[%s] Data Link OAuth - Failed to build OAuth Token. (Settings: %s) Exchange Code Response: %s")
+			, Executor->GetContextName().GetData()
+			, *OAuthSettings->GetName()
+			, InResponse.GetData());
+
+		Executor->FailNode(this);
+		return;
+	}
+
+	const FDataLinkOutputDataViewer& OutputDataViewer = NodeInstance.GetOutputDataViewer();
+	
+	// Copy Input Http Settings into Output
+	FDataLinkHttpSettings& OutputHttpSettings = OutputDataViewer.Get<FDataLinkHttpSettings>(UE::DataLink::OutputDefault);
+	OutputHttpSettings = InputDataViewer.Get<FDataLinkHttpSettings>(UE::DataLinkOAuth::InputHttp);
+
+	if (!OAuthSettings->AuthorizeHttpRequest(OAuthToken, OutputHttpSettings))
+	{
+		UE_LOG(LogDataLinkOAuth, Error, TEXT("[%s] Data Link OAuth - OAuth Token unexpectedly failed to authorize the Http Request. (Settings: %s) Exchange Code Response: %s")
+			, Executor->GetContextName().GetData()
+			, *OAuthSettings->GetName()
+			, InResponse.GetData());
+
+		Executor->FailNode(this);
+		return;
+	}
+
+	if (UDataLinkOAuthSubsystem* OAuthSubsystem = UDataLinkOAuthSubsystem::Get())
+	{
+		OAuthSubsystem->RegisterToken(OAuthSettings, OAuthToken);
+	}
+
+	Executor->SucceedNode(this, FConstStructView::Make(OutputHttpSettings));
+}
+
+#undef LOCTEXT_NAMESPACE
diff --git a/Engine/Plugins/Experimental/DataLinkOAuth/Source/DataLinkOAuth/Private/DataLinkOAuthAuthRequest.cpp b/Engine/Plugins/Experimental/DataLinkOAuth/Source/DataLinkOAuth/Private/DataLinkOAuthAuthRequest.cpp
new file mode 100644
--- /dev/null
+++ b/Engine/Plugins/Experimental/DataLinkOAuth/Source/DataLinkOAuth/Private/DataLinkOAuthAuthRequest.cpp
@@ -0,0 +1,128 @@
+// Copyright Epic Games, Inc. All Rights Reserved.
+
+#include "DataLinkOAuthAuthRequest.h"
+#include "DataLinkNodeOAuth.h"
+#include "DataLinkOAuthHandle.h"
+#include "DataLinkOAuthInstance.h"
+#include "DataLinkOAuthLog.h"
+#include "DataLinkOAuthSettings.h"
+#include "DataLinkOAuthSubsystem.h"
+#include "HttpServerModule.h"
+#include "IHttpRouter.h"
+#include "SocketSubsystem.h"
+#include "Sockets.h"
+
+namespace UE::DataLinkOAuth::Private
+{
+	TUniquePtr<FHttpServerResponse> CreateSuccessResponse()
+	{
+		constexpr const TCHAR* ResponseString = TEXT("<html>"\
+			"<head><meta http-equiv='refresh' content='3;url=http://unrealengine.com/'>""</head>"\
+			"<body>Please return to the App.</body>"\
+			"</html>");
+
+		TUniquePtr<FHttpServerResponse> Response = MakeUnique<FHttpServerResponse>();
+		Response->Code = EHttpServerResponseCodes::Ok;
+
+		const FTCHARToUTF8 ConvertToUtf8(ResponseString);
+		Response->Body.Append(reinterpret_cast<const uint8*>(ConvertToUtf8.Get()), ConvertToUtf8.Length());
+		return Response;
+	}
+
+	bool FindUnusedPort(int32& OutPort)
+	{
+		ISocketSubsystem* SocketSubsystem = ISocketSubsystem::Get(PLATFORM_SOCKETSUBSYSTEM);
+		if (!SocketSubsystem)
+		{
+			UE_LOG(LogDataLinkOAuth, Error, TEXT("Find Unused Port failed. Invalid Socket Subsystem"));
+			return false;
+		}
+
+		// Create address with port 0. After socket has been bound FSocket::GetAddress should give the actual port assigned
+		TSharedRef<FInternetAddr> InternetAddress = SocketSubsystem->CreateInternetAddr();
+		InternetAddress->SetLoopbackAddress();
+		InternetAddress->SetPort(0);
+
+		constexpr const TCHAR* SocketDescription = TEXT("Data Link OAuth - FindUnusedPort");
+		FSocket* Socket = SocketSubsystem->CreateSocket(NAME_Stream, SocketDescription, InternetAddress->GetProtocolType());
+		if (!Socket)
+		{
+			UE_LOG(LogDataLinkOAuth, Error, TEXT("Find Unused Port failed. Could not create Socket [%s, %s]")
+				, SocketDescription
+				, *InternetAddress->ToString(/*bAppendPort*/false));
+			return false;
+		}
+
+		if (!Socket->Bind(*InternetAddress))
+		{
+			UE_LOG(LogDataLinkOAuth, Error, TEXT("Find Unused Port failed. Failed to Bind Socket [%s, %s]")
+				, SocketDescription
+				, *InternetAddress->ToString(/*bAppendPort*/false));
+			return false;
+		}
+
+		// Re-use InternetAddress to get the port
+		Socket->GetAddress(*InternetAddress);
+		OutPort = InternetAddress->GetPort();
+		Socket->Close();
+		return true;
+	}
+
+	FDataLinkOAuthHandle StartListening(int32 InListenPort, const FOnAuthResponse& InOnAuthResponse)
+	{
+		const TSharedPtr<IHttpRouter> Router = FHttpServerModule::Get().GetHttpRouter(InListenPort, /*bFailOnBindFailure*/true);
+		if (!Router.IsValid())
+		{
+			UE_LOG(LogDataLinkOAuth, Error, TEXT("Start Listening failed. Could not get HTTP Router for port %d."), InListenPort);
+			return FDataLinkOAuthHandle();
+		}
+
+		UDataLinkOAuthSubsystem::FListenInstance ListenInstance;
+		ListenInstance.RouterWeak = Router;
+		ListenInstance.RequestPreprocessorHandle = Router->RegisterRequestPreprocessor(FHttpRequestHandler::CreateLambda(
+			[OnAuthResponse = InOnAuthResponse](const FHttpServerRequest& InRequest, const FHttpResultCallback& InOnComplete)->bool
+			{
+				InOnComplete(CreateSuccessResponse());
+				OnAuthResponse.ExecuteIfBound(InRequest);
+				return true;
+			}));
+
+		return UDataLinkOAuthSubsystem::Get()->RegisterListenInstance(MoveTemp(ListenInstance));
+	}
+}
+
+bool UE::DataLinkOAuth::RequestAuthorization(const FAuthRequestParams& InParams)
+{
+	UDataLinkOAuthSubsystem* OAuthSubsystem = UDataLinkOAuthSubsystem::Get();
+	if (!OAuthSubsystem)
+	{
+		UE_LOG(LogDataLinkOAuth, Error, TEXT("Start Listening failed. Invalid OAuth Subsystem."));
+		return false;
+	}
+
+	int32 ListenPort;
+	if (! Private::FindUnusedPort(ListenPort))
+	{
+		UE_LOG(LogDataLinkOAuth, Error, TEXT("Start Listening failed. Could not find unused port."));
+		return false;
+	}
+
+	const FDataLinkOAuthHandle ListenHandle = Private::StartListening(ListenPort, InParams.OnAuthResponse);
+	if (!ListenHandle.IsValid())
+	{
+		return false;
+	}
+
+	FDataLinkNodeOAuthInstance& OAuthInstance = InParams.OAuthInstanceView.Get();
+	OAuthInstance.ListenHandle = ListenHandle;
+	OAuthInstance.ListenPort = ListenPort;
+
+	UDataLinkOAuthSettings::FUrlBuilder RequestUrl;
+	check(InParams.OAuthSettings);
+	InParams.OAuthSettings->BuildAuthRequestUrl(RequestUrl, OAuthInstance);
+
+	FHttpServerModule::Get().StartAllListeners();
+
+	FPlatformProcess::LaunchURL(*RequestUrl, nullptr, nullptr);
+	return true;
+}
diff --git a/Engine/Plugins/Experimental/DataLinkOAuth/Source/DataLinkOAuth/Private/DataLinkOAuthAuthRequest.h b/Engine/Plugins/Experimental/DataLinkOAuth/Source/DataLinkOAuth/Private/DataLinkOAuthAuthRequest.h
new file mode 100644
--- /dev/null
+++ b/Engine/Plugins/Experimental/DataLinkOAuth/Source/DataLinkOAuth/Private/DataLinkOAuthAuthRequest.h
@@ -0,0 +1,30 @@
+// Copyright Epic Games, Inc. All Rights Reserved.
+
+#pragma once
+
+#include "Delegates/Delegate.h"
+#include "StructUtils/StructView.h"
+
+class FDataLinkOAuthHandle;
+class UDataLinkOAuthSettings;
+struct FDataLinkNodeOAuthInstance;
+struct FHttpServerRequest;
+
+namespace UE::DataLinkOAuth
+{
+	using FOnAuthResponse = TDelegate<void(const FHttpServerRequest&)>;
+
+	struct FAuthRequestParams
+	{
+		/** Settings to used for Authorization */
+		UDataLinkOAuthSettings* OAuthSettings;
+
+		/** OAuth Instance Data to Read from and Write to */
+		TStructView<FDataLinkNodeOAuthInstance> OAuthInstanceView;
+
+		/** Called when the Authorization has been granted*/
+		FOnAuthResponse OnAuthResponse;
+	};
+
+	bool RequestAuthorization(const FAuthRequestParams& InParams);
+}
diff --git a/Engine/Plugins/Experimental/DataLinkOAuth/Source/DataLinkOAuth/Private/DataLinkOAuthCodeAccessExchange.cpp b/Engine/Plugins/Experimental/DataLinkOAuth/Source/DataLinkOAuth/Private/DataLinkOAuthCodeAccessExchange.cpp
new file mode 100644
--- /dev/null
+++ b/Engine/Plugins/Experimental/DataLinkOAuth/Source/DataLinkOAuth/Private/DataLinkOAuthCodeAccessExchange.cpp
@@ -0,0 +1,68 @@
+// Copyright Epic Games, Inc. All Rights Reserved.
+
+#include "DataLinkOAuthCodeAccessExchange.h"
+#include "DataLinkNodeOAuth.h"
+#include "DataLinkOAuthInstance.h"
+#include "DataLinkOAuthLog.h"
+#include "DataLinkOAuthSettings.h"
+#include "HttpModule.h"
+#include "Interfaces/IHttpRequest.h"
+#include "Interfaces/IHttpResponse.h"
+
+namespace UE::DataLinkOAuth::Private
+{
+	FString GetResponseString(FHttpResponsePtr InResponse, bool bInProcessedSuccessfully)
+	{
+		if (!bInProcessedSuccessfully || !InResponse.IsValid())
+		{
+			UE_LOG(LogDataLinkOAuth, Error, TEXT("OAuth Code Exchange failed to be processed."));
+			return FString();
+		}
+
+		const int32 ResponseCode = InResponse->GetResponseCode();
+		if (!EHttpResponseCodes::IsOk(ResponseCode))
+		{
+			UE_LOG(LogDataLinkOAuth, Error, TEXT("OAuth Code Exchange failed with response code %d"), ResponseCode);
+			return FString();
+		}
+
+		return InResponse->GetContentAsString();
+	}
+}
+
+bool UE::DataLinkOAuth::ExchangeAuthCodeForAccess(const FExchangeAuthCodeParams& InParams)
+{
+	if (InParams.AuthCodeView.IsEmpty())
+	{
+		UE_LOG(LogDataLinkOAuth, Error, TEXT("OAuth Code Exchange Failed. Invalid Auth Code."));
+		return false;
+	}
+
+	const TSharedRef<IHttpRequest> HttpRequest = FHttpModule::Get().CreateRequest();
+
+	FDataLinkNodeOAuthInstance& OAuthInstance = InParams.OAuthInstanceView.Get();
+
+	UDataLinkOAuthSettings::FUrlBuilder RequestUrl;
+	check(InParams.OAuthSettings);
+
+	if (!InParams.OAuthSettings->BuildExchangeCodeTokenUrl(RequestUrl, OAuthInstance, InParams.AuthCodeView))
+	{
+		UE_LOG(LogDataLinkOAuth, Error, TEXT("OAuth Code Exchange Failed. Could not build request URL. (Settings: %s)")
+			, *InParams.OAuthSettings->GetName());
+		return false;
+	}
+
+	HttpRequest->SetVerb(TEXT("POST"));
+	HttpRequest->SetURL(*RequestUrl);
+	HttpRequest->SetHeader(TEXT("Content-Type"), TEXT("application/x-www-form-urlencoded"));
+	HttpRequest->SetDelegateThreadPolicy(EHttpRequestDelegateThreadPolicy::CompleteOnGameThread);
+
+	HttpRequest->OnProcessRequestComplete().BindLambda(
+		[OnResponse = InParams.OnResponse](FHttpRequestPtr InRequest, FHttpResponsePtr InResponse, bool bInProcessedSuccessfully)
+		{
+			const FString ResponseString = Private::GetResponseString(InResponse, bInProcessedSuccessfully);
+			OnResponse.ExecuteIfBound(ResponseString);
+		});
+
+	return HttpRequest->ProcessRequest();
+}
diff --git a/Engine/Plugins/Experimental/DataLinkOAuth/Source/DataLinkOAuth/Private/DataLinkOAuthCodeAccessExchange.h b/Engine/Plugins/Experimental/DataLinkOAuth/Source/DataLinkOAuth/Private/DataLinkOAuthCodeAccessExchange.h
new file mode 100644
--- /dev/null
+++ b/Engine/Plugins/Experimental/DataLinkOAuth/Source/DataLinkOAuth/Private/DataLinkOAuthCodeAccessExchange.h
@@ -0,0 +1,32 @@
+// Copyright Epic Games, Inc. All Rights Reserved.
+
+#pragma once
+
+#include "Containers/StringView.h"
+#include "Delegates/Delegate.h"
+#include "StructUtils/StructView.h"
+
+class UDataLinkOAuthSettings;
+struct FDataLinkNodeOAuthInstance;
+
+namespace UE::DataLinkOAuth
+{
+	using FOnExchangeAuthCodeResponse = TDelegate<void(FStringView InResponseString)>;
+
+	struct FExchangeAuthCodeParams
+	{
+		/** Settings to used for Authorization */
+		UDataLinkOAuthSettings* OAuthSettings;
+
+		/** OAuth Instance Data to Read from and Write to */
+		TStructView<FDataLinkNodeOAuthInstance> OAuthInstanceView;
+
+		/** Auth Code to exchange for Access Token */
+		FStringView AuthCodeView;
+
+		/** Called when the Exchange Request has received a response */
+		FOnExchangeAuthCodeResponse OnResponse;
+	};
+
+	bool ExchangeAuthCodeForAccess(const FExchangeAuthCodeParams& InParams);
+}
diff --git a/Engine/Plugins/Experimental/DataLinkOAuth/Source/DataLinkOAuth/Private/DataLinkOAuthDefaultSettings.cpp b/Engine/Plugins/Experimental/DataLinkOAuth/Source/DataLinkOAuth/Private/DataLinkOAuthDefaultSettings.cpp
new file mode 100644
--- /dev/null
+++ b/Engine/Plugins/Experimental/DataLinkOAuth/Source/DataLinkOAuth/Private/DataLinkOAuthDefaultSettings.cpp
@@ -0,0 +1,66 @@
+// Copyright Epic Games, Inc. All Rights Reserved.
+
+#include "DataLinkOAuthDefaultSettings.h"
+#include "DataLinkOAuthDefaultSharedData.h"
+#include "DataLinkOAuthInstance.h"
+#include "DataLinkOAuthLog.h"
+#include "HttpServerRequest.h"
+#include "PlatformHttp.h"
+#include "StructUtils/StructView.h"
+
+UDataLinkOAuthDefaultSettings::UDataLinkOAuthDefaultSettings()
+{
+	SharedDataType = FDataLinkOAuthDefaultSharedData::StaticStruct();
+}
+
+bool UDataLinkOAuthDefaultSettings::BuildAuthRequestUrl(FUrlBuilder& OutRequestUrl, FDataLinkNodeOAuthInstance& InOAuthInstance) const
+{
+	FDataLinkOAuthDefaultSharedData& SharedData = InOAuthInstance.SharedData.GetMutable<FDataLinkOAuthDefaultSharedData>();
+
+	const FString RedirectUriWithPort = FString::Printf(TEXT("%s:%s"), LoopbackAddress, *FString::FromInt(InOAuthInstance.ListenPort));
+
+	OutRequestUrl << AuthorizationURL
+		<< TEXT("?prompt=consent"\
+			"&response_type=code"\
+			"&access_type=offline")
+		<< TEXT("&redirect_uri=") << FPlatformHttp::UrlEncode(RedirectUriWithPort)
+		<< TEXT("&state=") << SharedData.State
+		<< TEXT("&client_id=") << ClientId;
+
+	if (!Scopes.IsEmpty())
+	{
+		OutRequestUrl << TEXT("&scope=");
+		for (const FString& Scope : Scopes)
+		{
+			OutRequestUrl << FPlatformHttp::UrlEncode(Scope) << TEXT("+");
+		}
+		// Remove last '+'
+		OutRequestUrl.RemoveAt(OutRequestUrl.Len() - 1, 1);
+	}
+
+	return true;
+}
+
+bool UDataLinkOAuthDefaultSettings::ValidateRequest(const FHttpServerRequest& InRequest, FDataLinkNodeOAuthInstance& InOAuthInstance) const
+{
+	FDataLinkOAuthDefaultSharedData& SharedData = InOAuthInstance.SharedData.GetMutable<FDataLinkOAuthDefaultSharedData>();
+
+	// Proceed only if the states match, ensuring that this instance made the request
+	// No node failure here, to expect request with the correct state to come
+	const FString* FoundState = InRequest.QueryParams.Find(TEXT("state"));
+	return FoundState && *FoundState == SharedData.State;
+}
+
+bool UDataLinkOAuthDefaultSettings::BuildExchangeCodeTokenUrl(FUrlBuilder& OutRequestUrl, FDataLinkNodeOAuthInstance& InOAuthInstance, FStringView InAuthCode) const
+{
+	const FString RedirectUriWithPort = FString::Printf(TEXT("%s:%s"), LoopbackAddress, *FString::FromInt(InOAuthInstance.ListenPort));
+
+	OutRequestUrl << TokenExchangeEndpoint
+		<< TEXT("?grant_type=authorization_code")
+		<< TEXT("&redirect_uri=") << FPlatformHttp::UrlEncode(RedirectUriWithPort)
+		<< TEXT("&client_id=") << ClientId
+		<< TEXT("&client_secret=") << ClientSecret
+		<< TEXT("&code=") << InAuthCode;
+
+	return true;
+}
diff --git a/Engine/Plugins/Experimental/DataLinkOAuth/Source/DataLinkOAuth/Private/DataLinkOAuthDefaultSharedData.cpp b/Engine/Plugins/Experimental/DataLinkOAuth/Source/DataLinkOAuth/Private/DataLinkOAuthDefaultSharedData.cpp
new file mode 100644
--- /dev/null
+++ b/Engine/Plugins/Experimental/DataLinkOAuth/Source/DataLinkOAuth/Private/DataLinkOAuthDefaultSharedData.cpp
@@ -0,0 +1,8 @@
+// Copyright Epic Games, Inc. All Rights Reserved.
+
+#include "DataLinkOAuthDefaultSharedData.h"
+
+FDataLinkOAuthDefaultSharedData::FDataLinkOAuthDefaultSharedData()
+	: State(FGuid::NewGuid().ToString())
+{
+}
diff --git a/Engine/Plugins/Experimental/DataLinkOAuth/Source/DataLinkOAuth/Private/DataLinkOAuthHandle.cpp b/Engine/Plugins/Experimental/DataLinkOAuth/Source/DataLinkOAuth/Private/DataLinkOAuthHandle.cpp
new file mode 100644
--- /dev/null
+++ b/Engine/Plugins/Experimental/DataLinkOAuth/Source/DataLinkOAuth/Private/DataLinkOAuthHandle.cpp
@@ -0,0 +1,39 @@
+// Copyright Epic Games, Inc. All Rights Reserved.
+
+#include "DataLinkOAuthHandle.h"
+#include <atomic>
+
+namespace UE::DataLinkOAuth::Private
+{
+	std::atomic<uint64> GNextId(0);
+
+	uint64 GenerateNewId()
+	{
+		uint64 NewId = ++GNextId;
+
+		// 0 is reserved for non-initialized state
+		if (NewId == 0)
+		{
+			NewId = ++GNextId;
+		}
+
+		return NewId;
+	}
+}
+
+FDataLinkOAuthHandle FDataLinkOAuthHandle::GenerateHandle()
+{
+	FDataLinkOAuthHandle Handle;
+	Handle.Id = UE::DataLinkOAuth::Private::GenerateNewId();
+	return Handle;
+}
+
+bool FDataLinkOAuthHandle::IsValid() const
+{
+	return Id != 0;
+}
+
+void FDataLinkOAuthHandle::Reset()
+{
+	Id = 0;
+}
diff --git a/Engine/Plugins/Experimental/DataLinkOAuth/Source/DataLinkOAuth/Private/DataLinkOAuthInstance.cpp b/Engine/Plugins/Experimental/DataLinkOAuth/Source/DataLinkOAuth/Private/DataLinkOAuthInstance.cpp
new file mode 100644
--- /dev/null
+++ b/Engine/Plugins/Experimental/DataLinkOAuth/Source/DataLinkOAuth/Private/DataLinkOAuthInstance.cpp
@@ -0,0 +1,21 @@
+// Copyright Epic Games, Inc. All Rights Reserved.
+
+#include "DataLinkOAuthInstance.h"
+#include "DataLinkOAuthSubsystem.h"
+
+FDataLinkNodeOAuthInstance::~FDataLinkNodeOAuthInstance()
+{
+	StopListening();
+}
+
+void FDataLinkNodeOAuthInstance::StopListening()
+{
+	if (ListenHandle.IsValid())
+	{
+		if (UDataLinkOAuthSubsystem* const OAuthSubsystem = UDataLinkOAuthSubsystem::Get())
+        {
+        	OAuthSubsystem->UnregisterListenInstance(ListenHandle);
+        }
+	}
+	ListenHandle.Reset();	
+}
diff --git a/Engine/Plugins/Experimental/DataLinkOAuth/Source/DataLinkOAuth/Private/DataLinkOAuthLog.h b/Engine/Plugins/Experimental/DataLinkOAuth/Source/DataLinkOAuth/Private/DataLinkOAuthLog.h
new file mode 100644
--- /dev/null
+++ b/Engine/Plugins/Experimental/DataLinkOAuth/Source/DataLinkOAuth/Private/DataLinkOAuthLog.h
@@ -0,0 +1,7 @@
+// Copyright Epic Games, Inc. All Rights Reserved.
+
+#pragma once
+
+#include "Logging/LogMacros.h"
+
+DECLARE_LOG_CATEGORY_EXTERN(LogDataLinkOAuth, Log, All);
diff --git a/Engine/Plugins/Experimental/DataLinkOAuth/Source/DataLinkOAuth/Private/DataLinkOAuthModule.cpp b/Engine/Plugins/Experimental/DataLinkOAuth/Source/DataLinkOAuth/Private/DataLinkOAuthModule.cpp
new file mode 100644
--- /dev/null
+++ b/Engine/Plugins/Experimental/DataLinkOAuth/Source/DataLinkOAuth/Private/DataLinkOAuthModule.cpp
@@ -0,0 +1,8 @@
+// Copyright Epic Games, Inc. All Rights Reserved.
+
+#include "DataLinkOAuthLog.h"
+#include "Modules/ModuleManager.h"
+
+DEFINE_LOG_CATEGORY(LogDataLinkOAuth);
+
+IMPLEMENT_MODULE(FDefaultModuleImpl, DataLinkOAuth)
diff --git a/Engine/Plugins/Experimental/DataLinkOAuth/Source/DataLinkOAuth/Private/DataLinkOAuthSettings.cpp b/Engine/Plugins/Experimental/DataLinkOAuth/Source/DataLinkOAuth/Private/DataLinkOAuthSettings.cpp
new file mode 100644
--- /dev/null
+++ b/Engine/Plugins/Experimental/DataLinkOAuth/Source/DataLinkOAuth/Private/DataLinkOAuthSettings.cpp
@@ -0,0 +1,73 @@
+// Copyright Epic Games, Inc. All Rights Reserved.
+
+#include "DataLinkOAuthSettings.h"
+#include "DataLinkHttpSettings.h"
+#include "DataLinkOAuthLog.h"
+#include "DataLinkOAuthToken.h"
+#include "DataLinkOAuthUtils.h"
+#include "Dom/JsonObject.h"
+#include "HttpServerRequest.h"
+#include "StructUtils/InstancedStruct.h"
+#include "StructUtils/StructView.h"
+
+FInstancedStruct UDataLinkOAuthSettings::MakeSharedData() const
+{
+	FInstancedStruct InstanceData;
+	InstanceData.InitializeAs(SharedDataType);
+	return InstanceData;
+}
+
+bool UDataLinkOAuthSettings::FindAuthCode(const FHttpServerRequest& InRequest, FDataLinkNodeOAuthInstance& InOAuthInstance, FStringView& OutAuthCodeView) const
+{
+	if (const FString* FoundCode = InRequest.QueryParams.Find(TEXT("code")))
+	{
+		OutAuthCodeView = *FoundCode;
+		return true;
+	}
+	return false;
+}
+
+bool UDataLinkOAuthSettings::BuildAuthToken(FStringView InAccessResponse, FDataLinkOAuthToken& OutAuthToken) const
+{
+	const TSharedPtr<FJsonObject> ResponseJson = UE::DataLinkOAuth::ResponseStringToJsonObject(InAccessResponse);
+	if (!ResponseJson.IsValid())
+	{
+		UE_LOG(LogDataLinkOAuth, Error, TEXT("Response %s could not be converted to a Json Object"), InAccessResponse.GetData());
+		return false;
+	}
+
+	if (!ResponseJson->TryGetStringField(TEXT("token_type"), OutAuthToken.TokenType))
+	{
+		UE_LOG(LogDataLinkOAuth, Error, TEXT("Response %s did not have a valid 'token_type' entry"), InAccessResponse.GetData());
+		return false;
+	}
+
+	if (!ResponseJson->TryGetStringField(TEXT("access_token"), OutAuthToken.AccessToken))
+	{
+		UE_LOG(LogDataLinkOAuth, Error, TEXT("Response %s did not have a valid 'access_token' entry"), InAccessResponse.GetData());
+		return false;
+	}
+
+	if (!ResponseJson->TryGetStringField(TEXT("refresh_token"), OutAuthToken.RefreshToken))
+	{
+		UE_LOG(LogDataLinkOAuth, Warning, TEXT("Response %s did not have a valid 'refresh_token' entry"), InAccessResponse.GetData());
+	}
+
+	uint32 ExpiresInSeconds = 0;
+	if (ResponseJson->TryGetNumberField(TEXT("expires_in"), ExpiresInSeconds))
+	{
+		OutAuthToken.ExpirationDate = FDateTime::UtcNow() + FTimespan(0, 0, ExpiresInSeconds);
+	}
+	else
+	{
+		UE_LOG(LogDataLinkOAuth, Warning, TEXT("Response %s did not have a valid 'expires_in' entry"), InAccessResponse.GetData());
+	}
+
+	return true;
+}
+
+bool UDataLinkOAuthSettings::AuthorizeHttpRequest(const FDataLinkOAuthToken& InAuthToken, FDataLinkHttpSettings& InOutHttpSettings) const
+{
+	InOutHttpSettings.Headers.Add(TEXT("Authorization"), FString::Printf(TEXT("%s %s"), *InAuthToken.TokenType, *InAuthToken.AccessToken));
+	return true;
+}
diff --git a/Engine/Plugins/Experimental/DataLinkOAuth/Source/DataLinkOAuth/Private/DataLinkOAuthSubsystem.cpp b/Engine/Plugins/Experimental/DataLinkOAuth/Source/DataLinkOAuth/Private/DataLinkOAuthSubsystem.cpp
new file mode 100644
--- /dev/null
+++ b/Engine/Plugins/Experimental/DataLinkOAuth/Source/DataLinkOAuth/Private/DataLinkOAuthSubsystem.cpp
@@ -0,0 +1,88 @@
+// Copyright Epic Games, Inc. All Rights Reserved.
+
+#include "DataLinkOAuthSubsystem.h"
+#include "DataLinkOAuthSettings.h"
+#include "Engine/Engine.h"
+#include "IHttpRouter.h"
+
+namespace UE::DataLinkOAuth::Private
+{
+	FDateTime GetPaddedNow()
+	{
+		constexpr int32 PaddingSeconds = 5;
+		return FDateTime::UtcNow() + FTimespan(0, 0, PaddingSeconds);
+	}
+}
+
+UDataLinkOAuthSubsystem* UDataLinkOAuthSubsystem::Get()
+{
+	if (!UObjectInitialized())
+	{
+		return nullptr;
+	}
+
+	if (GEngine)
+	{
+		return GEngine->GetEngineSubsystem<UDataLinkOAuthSubsystem>();
+	}
+	return nullptr;
+}
+
+FDataLinkOAuthHandle UDataLinkOAuthSubsystem::RegisterListenInstance(FListenInstance&& InInstance)
+{
+	FDataLinkOAuthHandle NewHandle = FDataLinkOAuthHandle::GenerateHandle();
+	ListeningInstances.Add(NewHandle, MoveTemp(InInstance));
+	return NewHandle;
+}
+
+void UDataLinkOAuthSubsystem::UnregisterListenInstance(FDataLinkOAuthHandle InHandle)
+{
+	FListenInstance ListenInstance;
+	if (!ListeningInstances.RemoveAndCopyValue(InHandle, ListenInstance))
+	{
+		return;
+	}
+
+	TSharedPtr<IHttpRouter> Router = ListenInstance.RouterWeak.Pin();
+	if (!Router.IsValid())
+	{
+		return;
+	}
+
+	Router->UnregisterRequestPreprocessor(ListenInstance.RequestPreprocessorHandle);
+}
+
+const FDataLinkOAuthToken* UDataLinkOAuthSubsystem::FindToken(const UDataLinkOAuthSettings* InOAuthSettings) const
+{
+	const FDateTime Now = UE::DataLinkOAuth::Private::GetPaddedNow();
+
+	const FDataLinkOAuthToken* Token = Tokens.Find(FDataLinkOAuthTokenHandle(InOAuthSettings));
+	if (Token && Now < Token->ExpirationDate)
+	{
+		return Token;	
+	}
+
+	return nullptr;
+}
+
+void UDataLinkOAuthSubsystem::RegisterToken(const UDataLinkOAuthSettings* InOAuthSettings, const FDataLinkOAuthToken& InToken)
+{
+	// Duplicate so that any changes to the OAuth Settings does not affect the cached token handle key
+	const UDataLinkOAuthSettings* DuplicateOAuthSettings = DuplicateObject(InOAuthSettings, this);
+
+	Tokens.Add(FDataLinkOAuthTokenHandle(DuplicateOAuthSettings), InToken);
+}
+
+void UDataLinkOAuthSubsystem::CleanExpiredTokens()
+{
+	const FDateTime Now = UE::DataLinkOAuth::Private::GetPaddedNow();
+
+	for (decltype(Tokens)::TIterator Iter(Tokens); Iter; ++Iter)
+	{
+		const FDataLinkOAuthToken& Token = Iter.Value();
+		if (Now >= Token.ExpirationDate)
+		{
+			Iter.RemoveCurrent();
+		}
+	}
+}
diff --git a/Engine/Plugins/Experimental/DataLinkOAuth/Source/DataLinkOAuth/Private/DataLinkOAuthTokenHandle.cpp b/Engine/Plugins/Experimental/DataLinkOAuth/Source/DataLinkOAuth/Private/DataLinkOAuthTokenHandle.cpp
new file mode 100644
--- /dev/null
+++ b/Engine/Plugins/Experimental/DataLinkOAuth/Source/DataLinkOAuth/Private/DataLinkOAuthTokenHandle.cpp
@@ -0,0 +1,83 @@
+// Copyright Epic Games, Inc. All Rights Reserved.
+
+#include "DataLinkOAuthTokenHandle.h"
+#include "DataLinkOAuthSettings.h"
+#include "Serialization/ArchiveObjectCrc32.h"
+
+namespace UE::DataLinkOAuth::Private
+{
+	bool IsTokenIndifferent(const FProperty* InProperty)
+	{
+		return InProperty->HasAllPropertyFlags(CPF_Transient);
+	}
+}
+
+FDataLinkOAuthTokenHandle::FDataLinkOAuthTokenHandle(const UDataLinkOAuthSettings* InOAuthSettings)
+	: OAuthSettings(InOAuthSettings)
+{
+	RecalculateHash();
+}
+
+bool FDataLinkOAuthTokenHandle::operator==(const FDataLinkOAuthTokenHandle& InOther) const
+{
+	if (CachedHash != InOther.CachedHash)
+	{
+		return false;
+	}
+
+	// If both are equal (unlikely) then no need to perform deep comparison
+	if (OAuthSettings == InOther.OAuthSettings)
+	{
+		return true;
+	}
+
+	if (!OAuthSettings || !InOther.OAuthSettings)
+	{
+		return false;
+	}
+
+	const UClass* const OAuthSettingsClass = OAuthSettings->GetClass();
+	if (OAuthSettingsClass != InOther.OAuthSettings->GetClass())
+	{
+		return false;
+	}
+
+	constexpr EPropertyPortFlags PortFlags = PPF_DeepComparison;
+
+	for (const FProperty* Property : TFieldRange<FProperty>(OAuthSettingsClass))
+	{
+		if (UE::DataLinkOAuth::Private::IsTokenIndifferent(Property))
+		{
+			continue;
+		}
+
+		for (int32 Index = 0; Index < Property->ArrayDim; ++Index)
+		{
+			if (!Property->Identical_InContainer(OAuthSettings, InOther.OAuthSettings, Index, PortFlags))
+			{
+				return false;
+			}
+		}
+	}
+
+	return true;
+}
+
+void FDataLinkOAuthTokenHandle::RecalculateHash()
+{
+	if (!OAuthSettings)
+	{
+		CachedHash = 0;
+		return;
+	}
+
+	class : public FArchiveObjectCrc32
+	{
+		virtual bool ShouldSkipProperty(const FProperty* InProperty) const override
+		{
+			return UE::DataLinkOAuth::Private::IsTokenIndifferent(InProperty);
+		}
+	} Archive;
+
+	CachedHash = Archive.Crc32(ConstCast(OAuthSettings), 0);
+}
diff --git a/Engine/Plugins/Experimental/DataLinkOAuth/Source/DataLinkOAuth/Private/DataLinkOAuthUtils.cpp b/Engine/Plugins/Experimental/DataLinkOAuth/Source/DataLinkOAuth/Private/DataLinkOAuthUtils.cpp
new file mode 100644
--- /dev/null
+++ b/Engine/Plugins/Experimental/DataLinkOAuth/Source/DataLinkOAuth/Private/DataLinkOAuthUtils.cpp
@@ -0,0 +1,25 @@
+// Copyright Epic Games, Inc. All Rights Reserved.
+
+#include "DataLinkOAuthUtils.h"
+#include "Containers/StringView.h"
+#include "Dom/JsonObject.h"
+#include "Dom/JsonValue.h"
+#include "Serialization/JsonReader.h"
+#include "Serialization/JsonSerializer.h"
+
+TSharedPtr<FJsonObject> UE::DataLinkOAuth::ResponseStringToJsonObject(FStringView InResponseString)
+{
+	TSharedPtr<FJsonValue> ResponseJson;
+	if (!FJsonSerializer::Deserialize(TJsonReaderFactory<>::Create(InResponseString.GetData()), ResponseJson))
+	{
+		return nullptr;
+	}
+
+	TSharedPtr<FJsonObject>* JsonObject;
+	if (!ResponseJson->TryGetObject(JsonObject))
+	{
+		return nullptr;
+	}
+
+	return *JsonObject;
+}
diff --git a/Engine/Plugins/Experimental/DataLinkOAuth/Source/DataLinkOAuth/Public/DataLinkNodeOAuth.h b/Engine/Plugins/Experimental/DataLinkOAuth/Source/DataLinkOAuth/Public/DataLinkNodeOAuth.h
new file mode 100644
--- /dev/null
+++ b/Engine/Plugins/Experimental/DataLinkOAuth/Source/DataLinkOAuth/Public/DataLinkNodeOAuth.h
@@ -0,0 +1,45 @@
+// Copyright Epic Games, Inc. All Rights Reserved.
+
+#pragma once
+
+#include "DataLinkNode.h"
+#include "DataLinkOAuthHandle.h"
+#include "DataLinkNodeOAuth.generated.h"
+
+class UDataLinkOAuthSettings;
+struct FHttpServerRequest;
+
+namespace UE::DataLinkOAuth
+{
+	const FLazyName InputHttp(TEXT("InputHttp"));
+	const FLazyName InputOAuth(TEXT("InputOAuth"));
+}
+
+USTRUCT(BlueprintType, DisplayName="Data Link OAuth Settings")
+struct FDataLinkOAuthSettingsWrapper
+{
+	GENERATED_BODY()
+
+	UPROPERTY(EditAnywhere, BlueprintReadWrite, Category="Data Link")
+	TObjectPtr<UDataLinkOAuthSettings> OAuthSettings;
+};
+
+UCLASS(MinimalAPI, DisplayName="OAuth", Category="Authorization")
+class UDataLinkNodeOAuth : public UDataLinkNode
+{
+	GENERATED_BODY()
+
+public:
+	DATALINKOAUTH_API UDataLinkNodeOAuth();
+
+protected:
+	//~ Begin UDataLinkNode
+	DATALINKOAUTH_API virtual void OnBuildPins(FDataLinkPinBuilder& Inputs, FDataLinkPinBuilder& Outputs) const override;
+	DATALINKOAUTH_API virtual EDataLinkExecutionReply OnExecute(FDataLinkExecutor& InExecutor) const override;
+	//~ End UDataLinkNode
+
+private:
+	void OnAuthResponse(const FHttpServerRequest& InRequest, TWeakPtr<FDataLinkExecutor> InExecutorWeak) const;
+
+	void OnExchangeCodeResponse(FStringView InResponse, TWeakPtr<FDataLinkExecutor> InExecutorWeak) const;
+};
diff --git a/Engine/Plugins/Experimental/DataLinkOAuth/Source/DataLinkOAuth/Public/DataLinkOAuthDefaultSettings.h b/Engine/Plugins/Experimental/DataLinkOAuth/Source/DataLinkOAuth/Public/DataLinkOAuthDefaultSettings.h
new file mode 100644
--- /dev/null
+++ b/Engine/Plugins/Experimental/DataLinkOAuth/Source/DataLinkOAuth/Public/DataLinkOAuthDefaultSettings.h
@@ -0,0 +1,31 @@
+// Copyright Epic Games, Inc. All Rights Reserved.
+
+#pragma once
+
+#include "DataLinkOAuthSettings.h"
+#include "DataLinkOAuthDefaultSettings.generated.h"
+
+UCLASS(MinimalAPI)
+class UDataLinkOAuthDefaultSettings : public UDataLinkOAuthSettings
+{
+	GENERATED_BODY()
+
+public:
+	DATALINKOAUTH_API UDataLinkOAuthDefaultSettings();
+
+	//~ Begin UDataLinkOAuthSettings
+	DATALINKOAUTH_API virtual bool BuildAuthRequestUrl(FUrlBuilder& OutRequestUrl, FDataLinkNodeOAuthInstance& InOAuthInstance) const override;
+	DATALINKOAUTH_API virtual bool ValidateRequest(const FHttpServerRequest& InRequest, FDataLinkNodeOAuthInstance& InOAuthInstance) const override;
+	DATALINKOAUTH_API virtual bool BuildExchangeCodeTokenUrl(FUrlBuilder& OutRequestUrl, FDataLinkNodeOAuthInstance& InOAuthInstance, FStringView InAuthCode) const override;
+	//~ End UDataLinkOAuthSettings
+
+private:
+	UPROPERTY(EditAnywhere, Category="OAuth")
+	FString AuthorizationURL;
+
+	UPROPERTY(EditAnywhere, Category="OAuth")
+	FString TokenExchangeEndpoint;
+
+	UPROPERTY(EditAnywhere, Category="OAuth")
+	TArray<FString> Scopes;
+};
diff --git a/Engine/Plugins/Experimental/DataLinkOAuth/Source/DataLinkOAuth/Public/DataLinkOAuthDefaultSharedData.h b/Engine/Plugins/Experimental/DataLinkOAuth/Source/DataLinkOAuth/Public/DataLinkOAuthDefaultSharedData.h
new file mode 100644
--- /dev/null
+++ b/Engine/Plugins/Experimental/DataLinkOAuth/Source/DataLinkOAuth/Public/DataLinkOAuthDefaultSharedData.h
@@ -0,0 +1,17 @@
+// Copyright Epic Games, Inc. All Rights Reserved.
+
+#pragma once
+
+#include "Containers/UnrealString.h"
+#include "DataLinkOAuthDefaultSharedData.generated.h"
+
+USTRUCT()
+struct FDataLinkOAuthDefaultSharedData
+{
+	GENERATED_BODY()
+
+	DATALINKOAUTH_API FDataLinkOAuthDefaultSharedData();
+
+	UPROPERTY(meta=(IgnoreForMemberInitializationTest))
+	FString State;
+};
diff --git a/Engine/Plugins/Experimental/DataLinkOAuth/Source/DataLinkOAuth/Public/DataLinkOAuthHandle.h b/Engine/Plugins/Experimental/DataLinkOAuth/Source/DataLinkOAuth/Public/DataLinkOAuthHandle.h
new file mode 100644
--- /dev/null
+++ b/Engine/Plugins/Experimental/DataLinkOAuth/Source/DataLinkOAuth/Public/DataLinkOAuthHandle.h
@@ -0,0 +1,33 @@
+// Copyright Epic Games, Inc. All Rights Reserved.
+
+#pragma once
+
+#include "HAL/Platform.h"
+#include "Templates/TypeHash.h"
+
+class FDataLinkOAuthHandle
+{
+public:
+	FDataLinkOAuthHandle() = default;
+
+	static FDataLinkOAuthHandle GenerateHandle();
+
+	/** Whether this handle was ever initialized */
+	DATALINKOAUTH_API bool IsValid() const;
+
+	/** Clear handle to indicate it is no longer used */
+	DATALINKOAUTH_API void Reset();
+
+	bool operator==(const FDataLinkOAuthHandle& InOther) const
+	{
+		return Id == InOther.Id;
+	}
+
+	friend uint32 GetTypeHash(const FDataLinkOAuthHandle& InHandle)
+	{
+		return GetTypeHash(InHandle.Id);
+	}
+
+private:
+	uint64 Id = 0;
+};
diff --git a/Engine/Plugins/Experimental/DataLinkOAuth/Source/DataLinkOAuth/Public/DataLinkOAuthInstance.h b/Engine/Plugins/Experimental/DataLinkOAuth/Source/DataLinkOAuth/Public/DataLinkOAuthInstance.h
new file mode 100644
--- /dev/null
+++ b/Engine/Plugins/Experimental/DataLinkOAuth/Source/DataLinkOAuth/Public/DataLinkOAuthInstance.h
@@ -0,0 +1,25 @@
+// Copyright Epic Games, Inc. All Rights Reserved.
+
+#pragma once
+
+#include "DataLinkOAuthHandle.h"
+#include "StructUtils/InstancedStruct.h"
+#include "DataLinkOAuthInstance.generated.h"
+
+USTRUCT()
+struct FDataLinkNodeOAuthInstance
+{
+	GENERATED_BODY()
+
+	DATALINKOAUTH_API ~FDataLinkNodeOAuthInstance();
+
+	void StopListening();
+
+	FDataLinkOAuthHandle ListenHandle;
+
+	UPROPERTY()
+	int32 ListenPort = 0;
+
+	UPROPERTY()
+	FInstancedStruct SharedData;
+};
diff --git a/Engine/Plugins/Experimental/DataLinkOAuth/Source/DataLinkOAuth/Public/DataLinkOAuthSettings.h b/Engine/Plugins/Experimental/DataLinkOAuth/Source/DataLinkOAuth/Public/DataLinkOAuthSettings.h
new file mode 100644
--- /dev/null
+++ b/Engine/Plugins/Experimental/DataLinkOAuth/Source/DataLinkOAuth/Public/DataLinkOAuthSettings.h
@@ -0,0 +1,100 @@
+// Copyright Epic Games, Inc. All Rights Reserved.
+
+#pragma once
+
+#include "Containers/StringFwd.h"
+#include "UObject/Object.h"
+#include "DataLinkOAuthSettings.generated.h"
+
+struct FDataLinkHttpSettings;
+struct FDataLinkNodeOAuthInstance;
+struct FDataLinkOAuthToken;
+struct FHttpServerRequest;
+struct FInstancedStruct;
+
+UCLASS(MinimalAPI, Abstract)
+class UDataLinkOAuthSettings : public UObject
+{
+	GENERATED_BODY()
+
+public:
+	using FUrlBuilder = TStringBuilder<512>;
+
+	/**
+	 * Instantiated data of type that matches the 'SharedDataType'.
+	 * This is data that will live throughout the OAuth process and used by stages to read from/write to
+	 */
+	FInstancedStruct MakeSharedData() const;
+
+	/**
+	 * Builds an Authorization URL
+	 * @param OutRequestUrl the resulting request URL
+	 * @param InOAuthInstance oauth instance data allowing different stages of OAuth to read from/write to
+	 * @return true if the request URL was successfully created
+	 */
+	virtual bool BuildAuthRequestUrl(FUrlBuilder& OutRequestUrl, FDataLinkNodeOAuthInstance& InOAuthInstance) const
+	{
+		return false;
+	}
+
+	/**
+	 * Validates that the request relates to our OAuth flow
+	 * @param InRequest the request containing data to validate against
+	 * @param InOAuthInstance oauth instance data containing data to use for validation
+	 * @return true if request is valid. False if it's a request to ignore
+	 */
+	virtual bool ValidateRequest(const FHttpServerRequest& InRequest, FDataLinkNodeOAuthInstance& InOAuthInstance) const
+	{
+		return true;
+	}
+
+	/**
+	 * Attempts to find the Auth Code within the given Request
+	 * @param InRequest the request potentially containing the authorization code to retrieve
+	 * @param InOAuthInstance oauth instance data allowing different stages of OAuth to read from/write to
+	 * @param OutAuthCodeView the auth code string view if it was found
+	 * @return true if the Auth Code was found
+	 */
+	DATALINKOAUTH_API virtual bool FindAuthCode(const FHttpServerRequest& InRequest, FDataLinkNodeOAuthInstance& InOAuthInstance, FStringView& OutAuthCodeView) const;
+
+	/**
+	 * Builds an Exchange Code for Access Token URL
+	 * @param OutRequestUrl the resulting request url
+	 * @param InOAuthInstance oauth instance data allowing different stages of OAuth to read from/write to
+	 * @param InAuthCode the authorization code to exchange
+	 * @return true if the Request URL was successfully created
+	 */
+	virtual bool BuildExchangeCodeTokenUrl(FUrlBuilder& OutRequestUrl, FDataLinkNodeOAuthInstance& InOAuthInstance, FStringView InAuthCode) const
+	{
+		return false;
+	}
+
+	/**
+	 * Builds an Auth Token based on the Access Response String
+	 * @param InAccessResponse the access response
+	 * @param OutAuthToken the auth token to build
+	 * @return true if the token was built
+	 */
+	DATALINKOAUTH_API virtual bool BuildAuthToken(FStringView InAccessResponse, FDataLinkOAuthToken& OutAuthToken) const;
+
+	/**
+	 * Adds Authorization header to the given Http Settings if the Access Response is valid
+	 * @param InAuthToken the auth token containing the authorization parameters
+	 * @param InOutHttpSettings the http settings to add authorization to
+	 * @return true if the authorization was added
+	 */
+	DATALINKOAUTH_API virtual bool AuthorizeHttpRequest(const FDataLinkOAuthToken& InAuthToken, FDataLinkHttpSettings& InOutHttpSettings) const;
+
+protected:
+	static constexpr const TCHAR* LoopbackAddress = TEXT("http://127.0.0.1");
+
+	/** Optional shared data type to use across the different stages of OAuth */
+	UPROPERTY()
+	TObjectPtr<const UScriptStruct> SharedDataType;
+
+	UPROPERTY(EditAnywhere, Category="Client")
+	FString ClientId;
+
+	UPROPERTY(EditAnywhere, Category="Client")
+	FString ClientSecret;
+};
diff --git a/Engine/Plugins/Experimental/DataLinkOAuth/Source/DataLinkOAuth/Public/DataLinkOAuthSubsystem.h b/Engine/Plugins/Experimental/DataLinkOAuth/Source/DataLinkOAuth/Public/DataLinkOAuthSubsystem.h
new file mode 100644
--- /dev/null
+++ b/Engine/Plugins/Experimental/DataLinkOAuth/Source/DataLinkOAuth/Public/DataLinkOAuthSubsystem.h
@@ -0,0 +1,43 @@
+// Copyright Epic Games, Inc. All Rights Reserved.
+
+#pragma once
+
+#include "Containers/Map.h"
+#include "DataLinkOAuthHandle.h"
+#include "DataLinkOAuthToken.h"
+#include "DataLinkOAuthTokenHandle.h"
+#include "Subsystems/EngineSubsystem.h"
+#include "DataLinkOAuthSubsystem.generated.h"
+
+class IHttpRouter;
+class UDataLinkOAuthSettings;
+
+UCLASS(MinimalAPI)
+class UDataLinkOAuthSubsystem : public UEngineSubsystem
+{
+	GENERATED_BODY()
+
+public:
+	static UDataLinkOAuthSubsystem* Get();
+
+	struct FListenInstance
+	{
+		TWeakPtr<IHttpRouter> RouterWeak;
+		FDelegateHandle RequestPreprocessorHandle;
+	};
+	FDataLinkOAuthHandle RegisterListenInstance(FListenInstance&& InInstance);
+
+	void UnregisterListenInstance(FDataLinkOAuthHandle InHandle);
+
+	const FDataLinkOAuthToken* FindToken(const UDataLinkOAuthSettings* InOAuthSettings) const;
+
+	void RegisterToken(const UDataLinkOAuthSettings* InOAuthSettings, const FDataLinkOAuthToken& InToken);
+
+	void CleanExpiredTokens();
+
+private:
+	TMap<FDataLinkOAuthHandle, FListenInstance> ListeningInstances;
+
+	UPROPERTY()
+	TMap<FDataLinkOAuthTokenHandle, FDataLinkOAuthToken> Tokens;
+};
diff --git a/Engine/Plugins/Experimental/DataLinkOAuth/Source/DataLinkOAuth/Public/DataLinkOAuthToken.h b/Engine/Plugins/Experimental/DataLinkOAuth/Source/DataLinkOAuth/Public/DataLinkOAuthToken.h
new file mode 100644
--- /dev/null
+++ b/Engine/Plugins/Experimental/DataLinkOAuth/Source/DataLinkOAuth/Public/DataLinkOAuthToken.h
@@ -0,0 +1,24 @@
+// Copyright Epic Games, Inc. All Rights Reserved.
+
+#pragma once
+
+#include "Containers/UnrealString.h"
+#include "DataLinkOAuthToken.generated.h"
+
+USTRUCT()
+struct FDataLinkOAuthToken
+{
+	GENERATED_BODY()
+
+	UPROPERTY()
+	FString TokenType;
+
+	UPROPERTY()
+	FString AccessToken;
+
+	UPROPERTY()
+	FString RefreshToken;
+
+	UPROPERTY()
+	FDateTime ExpirationDate;
+};
diff --git a/Engine/Plugins/Experimental/DataLinkOAuth/Source/DataLinkOAuth/Public/DataLinkOAuthTokenHandle.h b/Engine/Plugins/Experimental/DataLinkOAuth/Source/DataLinkOAuth/Public/DataLinkOAuthTokenHandle.h
new file mode 100644
--- /dev/null
+++ b/Engine/Plugins/Experimental/DataLinkOAuth/Source/DataLinkOAuth/Public/DataLinkOAuthTokenHandle.h
@@ -0,0 +1,33 @@
+// Copyright Epic Games, Inc. All Rights Reserved.
+
+#pragma once
+
+#include "DataLinkOAuthTokenHandle.generated.h"
+
+class UDataLinkOAuthSettings;
+
+USTRUCT()
+struct FDataLinkOAuthTokenHandle
+{
+	GENERATED_BODY()
+
+	FDataLinkOAuthTokenHandle() = default;
+
+	explicit FDataLinkOAuthTokenHandle(const UDataLinkOAuthSettings* InOAuthSettings);
+
+	bool operator==(const FDataLinkOAuthTokenHandle& InOther) const;
+
+	friend uint32 GetTypeHash(const FDataLinkOAuthTokenHandle& InHandle)
+	{
+		return InHandle.CachedHash;
+	}
+
+private:
+	void RecalculateHash();
+
+	UPROPERTY()
+	uint32 CachedHash = 0;
+
+	UPROPERTY()
+	TObjectPtr<const UDataLinkOAuthSettings> OAuthSettings;
+};
diff --git a/Engine/Plugins/Experimental/DataLinkOAuth/Source/DataLinkOAuth/Public/DataLinkOAuthUtils.h b/Engine/Plugins/Experimental/DataLinkOAuth/Source/DataLinkOAuth/Public/DataLinkOAuthUtils.h
new file mode 100644
--- /dev/null
+++ b/Engine/Plugins/Experimental/DataLinkOAuth/Source/DataLinkOAuth/Public/DataLinkOAuthUtils.h
@@ -0,0 +1,13 @@
+// Copyright Epic Games, Inc. All Rights Reserved.
+
+#pragma once
+
+#include "Containers/StringFwd.h"
+#include "Templates/SharedPointerFwd.h"
+
+class FJsonObject;
+
+namespace UE::DataLinkOAuth
+{
+	DATALINKOAUTH_API TSharedPtr<FJsonObject> ResponseStringToJsonObject(FStringView InResponseString);
+}
diff --git a/Engine/Plugins/Experimental/DataLinkOAuth/Source/DataLinkOAuthEditor/DataLinkOAuthEditor.Build.cs b/Engine/Plugins/Experimental/DataLinkOAuth/Source/DataLinkOAuthEditor/DataLinkOAuthEditor.Build.cs
new file mode 100644
--- /dev/null
+++ b/Engine/Plugins/Experimental/DataLinkOAuth/Source/DataLinkOAuthEditor/DataLinkOAuthEditor.Build.cs
@@ -0,0 +1,25 @@
+// Copyright Epic Games, Inc. All Rights Reserved.
+
+using UnrealBuildTool;
+
+public class DataLinkOAuthEditor : ModuleRules
+{
+    public DataLinkOAuthEditor(ReadOnlyTargetRules Target) : base(Target)
+    {
+        PCHUsage = ModuleRules.PCHUsageMode.UseExplicitOrSharedPCHs;
+
+        PrivateDependencyModuleNames.AddRange(
+            new string[]
+            {
+                "AssetDefinition",
+                "AssetTools",
+                "Core",
+                "CoreUObject",
+                "ClassViewer",
+                "DataLinkOAuth",
+                "Engine",
+                "UnrealEd",
+            }
+        );
+    }
+}
\ No newline at end of file
diff --git a/Engine/Plugins/Experimental/DataLinkOAuth/Source/DataLinkOAuthEditor/Private/AssetDefinition_DataLinkOAuthSettings.cpp b/Engine/Plugins/Experimental/DataLinkOAuth/Source/DataLinkOAuthEditor/Private/AssetDefinition_DataLinkOAuthSettings.cpp
new file mode 100644
--- /dev/null
+++ b/Engine/Plugins/Experimental/DataLinkOAuth/Source/DataLinkOAuthEditor/Private/AssetDefinition_DataLinkOAuthSettings.cpp
@@ -0,0 +1,23 @@
+// Copyright Epic Games, Inc. All Rights Reserved.
+
+#include "AssetDefinition_DataLinkOAuthSettings.h"
+#include "DataLinkOAuthSettings.h"
+
+#define LOCTEXT_NAMESPACE "AssetDefinition_DataLinkOAuthSettings"
+
+FText UAssetDefinition_DataLinkOAuthSettings::GetAssetDisplayName() const
+{
+	return LOCTEXT("AssetDisplayName", "Data Link OAuth Settings");
+}
+
+FLinearColor UAssetDefinition_DataLinkOAuthSettings::GetAssetColor() const
+{
+	return FLinearColor(FColor(30,144,255));
+}
+
+TSoftClassPtr<UObject> UAssetDefinition_DataLinkOAuthSettings::GetAssetClass() const
+{
+	return UDataLinkOAuthSettings::StaticClass();
+}
+
+#undef LOCTEXT_NAMESPACE
diff --git a/Engine/Plugins/Experimental/DataLinkOAuth/Source/DataLinkOAuthEditor/Private/AssetDefinition_DataLinkOAuthSettings.h b/Engine/Plugins/Experimental/DataLinkOAuth/Source/DataLinkOAuthEditor/Private/AssetDefinition_DataLinkOAuthSettings.h
new file mode 100644
--- /dev/null
+++ b/Engine/Plugins/Experimental/DataLinkOAuth/Source/DataLinkOAuthEditor/Private/AssetDefinition_DataLinkOAuthSettings.h
@@ -0,0 +1,19 @@
+// Copyright Epic Games, Inc. All Rights Reserved.
+
+#pragma once
+
+#include "AssetDefinitionDefault.h"
+#include "AssetDefinition_DataLinkOAuthSettings.generated.h"
+
+UCLASS()
+class UAssetDefinition_DataLinkOAuthSettings : public UAssetDefinitionDefault
+{
+	GENERATED_BODY()
+
+public:
+	//~ Begin UAssetDefinition
+	virtual FText GetAssetDisplayName() const override;
+	virtual FLinearColor GetAssetColor() const override;
+	virtual TSoftClassPtr<UObject> GetAssetClass() const override;
+	//~ End UAssetDefinition
+};
diff --git a/Engine/Plugins/Experimental/DataLinkOAuth/Source/DataLinkOAuthEditor/Private/DataLinkOAuthEditorModule.cpp b/Engine/Plugins/Experimental/DataLinkOAuth/Source/DataLinkOAuthEditor/Private/DataLinkOAuthEditorModule.cpp
new file mode 100644
--- /dev/null
+++ b/Engine/Plugins/Experimental/DataLinkOAuth/Source/DataLinkOAuthEditor/Private/DataLinkOAuthEditorModule.cpp
@@ -0,0 +1,5 @@
+// Copyright Epic Games, Inc. All Rights Reserved.
+
+#include "Modules/ModuleManager.h"
+
+IMPLEMENT_MODULE(FDefaultModuleImpl, DataLinkOAuthEditor)
diff --git a/Engine/Plugins/Experimental/DataLinkOAuth/Source/DataLinkOAuthEditor/Private/DataLinkOAuthSettingsFactory.cpp b/Engine/Plugins/Experimental/DataLinkOAuth/Source/DataLinkOAuthEditor/Private/DataLinkOAuthSettingsFactory.cpp
new file mode 100644
--- /dev/null
+++ b/Engine/Plugins/Experimental/DataLinkOAuth/Source/DataLinkOAuthEditor/Private/DataLinkOAuthSettingsFactory.cpp
@@ -0,0 +1,96 @@
+// Copyright Epic Games, Inc. All Rights Reserved.
+
+#include "DataLinkOAuthSettingsFactory.h"
+#include "AssetTypeCategories.h"
+#include "ClassViewerFilter.h"
+#include "ClassViewerModule.h"
+#include "DataLinkOAuthSettings.h"
+#include "Kismet2/SClassPickerDialog.h"
+#include "Modules/ModuleManager.h"
+
+#define LOCTEXT_NAMESPACE "DataLinkOAuthSettingsFactory"
+
+namespace UE::DataLinkOAuthEditor::Private
+{
+	class FAssetClassParentFilter : public IClassViewerFilter
+	{
+	public:
+		/** All children of these classes will be included unless filtered out by another setting. */
+		TSet<const UClass*> AllowedChildrenOfClasses;
+
+		/** Disallowed class flags. */
+		EClassFlags DisallowedClassFlags = CLASS_None;
+
+		//~ Begin IClassViewerFilter
+		virtual bool IsClassAllowed(const FClassViewerInitializationOptions& InInitOptions, const UClass* InClass, TSharedRef<FClassViewerFilterFuncs> InFilterFunctions) override
+		{
+			return !InClass->HasAnyClassFlags(DisallowedClassFlags)
+				&& InFilterFunctions->IfInChildOfClassesSet(AllowedChildrenOfClasses, InClass) != EFilterReturn::Failed;
+		}
+		virtual bool IsUnloadedClassAllowed(const FClassViewerInitializationOptions& InInitOptions, const TSharedRef<const IUnloadedBlueprintData> InUnloadedClassData, TSharedRef<FClassViewerFilterFuncs> InFilterFunctions) override
+		{
+			return !InUnloadedClassData->HasAnyClassFlags(DisallowedClassFlags)
+				&& InFilterFunctions->IfInChildOfClassesSet(AllowedChildrenOfClasses, InUnloadedClassData) != EFilterReturn::Failed;
+		}
+		//~ End IClassViewerFilter
+	};
+}
+
+UDataLinkOAuthSettingsFactory::UDataLinkOAuthSettingsFactory()
+{
+	SupportedClass = UDataLinkOAuthSettings::StaticClass();
+	bCreateNew = true;
+	bEditorImport = false;
+	bEditAfterNew = true;
+}
+
+uint32 UDataLinkOAuthSettingsFactory::GetMenuCategories() const
+{
+	return EAssetTypeCategories::Misc;
+}
+
+bool UDataLinkOAuthSettingsFactory::ConfigureProperties()
+{
+	using namespace UE::DataLinkOAuthEditor;
+
+	FClassViewerModule& ClassViewerModule = FModuleManager::LoadModuleChecked<FClassViewerModule>("ClassViewer");
+
+	const TSharedRef<Private::FAssetClassParentFilter> Filter = MakeShared<Private::FAssetClassParentFilter>();
+	Filter->DisallowedClassFlags = CLASS_Abstract | CLASS_Deprecated | CLASS_NewerVersionExists | CLASS_HideDropDown;
+	Filter->AllowedChildrenOfClasses.Add(UDataLinkOAuthSettings::StaticClass());
+
+	FClassViewerInitializationOptions Options;
+	Options.Mode = EClassViewerMode::ClassPicker;
+	Options.NameTypeToDisplay = EClassViewerNameTypeToDisplay::DisplayName;
+	Options.ClassFilters.Add(Filter);
+
+	const FText TitleText = LOCTEXT("OAuthSettingsPickerTitle", "Pick OAuth Settings Class");
+
+	OAuthSettingsClass = nullptr;
+
+	UClass* ChosenClass = nullptr;
+	if (SClassPickerDialog::PickClass(TitleText, Options, ChosenClass, UDataLinkOAuthSettings::StaticClass()))
+	{
+		OAuthSettingsClass = ChosenClass;
+		return true;
+	}
+
+	return false;
+}
+
+UObject* UDataLinkOAuthSettingsFactory::FactoryCreateNew(UClass* InClass, UObject* InParent, FName InName, EObjectFlags InFlags, UObject* InContext, FFeedbackContext* InWarn)
+{
+	if (OAuthSettingsClass)
+	{
+		return NewObject<UDataLinkOAuthSettings>(InParent, OAuthSettingsClass, InName, InFlags);
+	}
+
+	if (ensureAlways(InClass && InClass->IsChildOf<UDataLinkOAuthSettings>()))
+	{
+		return NewObject<UDataLinkOAuthSettings>(InParent, InClass, InName, InFlags);
+	}
+
+	return nullptr;
+}
+
+#undef LOCTEXT_NAMESPACE
diff --git a/Engine/Plugins/Experimental/DataLinkOAuth/Source/DataLinkOAuthEditor/Private/DataLinkOAuthSettingsFactory.h b/Engine/Plugins/Experimental/DataLinkOAuth/Source/DataLinkOAuthEditor/Private/DataLinkOAuthSettingsFactory.h
new file mode 100644
--- /dev/null
+++ b/Engine/Plugins/Experimental/DataLinkOAuth/Source/DataLinkOAuthEditor/Private/DataLinkOAuthSettingsFactory.h
@@ -0,0 +1,28 @@
+// Copyright Epic Games, Inc. All Rights Reserved.
+
+#pragma once
+
+#include "Factories/Factory.h"
+#include "DataLinkOAuthSettingsFactory.generated.h"
+
+class UDataLinkOAuthSettings;
+
+UCLASS()
+class UDataLinkOAuthSettingsFactory : public UFactory
+{
+	GENERATED_BODY()
+
+public:
+	UDataLinkOAuthSettingsFactory();
+
+	//~ Begin UFactory
+	virtual uint32 GetMenuCategories() const override;
+	virtual bool ConfigureProperties() override;
+	virtual UObject* FactoryCreateNew(UClass* InClass, UObject* InParent, FName InName, EObjectFlags InFlags, UObject* InContext, FFeedbackContext* InWarn) override;
+	//~ End UFactory
+
+private:
+	/** The OAuth Settings Type to create */
+	UPROPERTY(EditAnywhere, Category="OAuth")
+	TSubclassOf<UDataLinkOAuthSettings> OAuthSettingsClass;
+};
