Motion Design: renamed Motion Design Data Link friendly name to Motion Design Data

diff --git a/Engine/Plugins/Experimental/AvalancheDataLink/AvalancheDataLink.uplugin b/Engine/Plugins/Experimental/AvalancheDataLink/AvalancheDataLink.uplugin
--- a/Engine/Plugins/Experimental/AvalancheDataLink/AvalancheDataLink.uplugin
+++ b/Engine/Plugins/Experimental/AvalancheDataLink/AvalancheDataLink.uplugin
@@ -2,7 +2,7 @@
 	"FileVersion": 3,
 	"Version": 1,
 	"VersionName": "1.0",
-	"FriendlyName": "Motion Design Data Link Integration",
+	"FriendlyName": "Motion Design Data Integration",
 	"Description": "",
 	"Category": "Other",
 	"CreatedBy": "",
diff --git a/Engine/Plugins/Experimental/DataLink/DataLink.uplugin b/Engine/Plugins/Experimental/DataLink/DataLink.uplugin
--- a/Engine/Plugins/Experimental/DataLink/DataLink.uplugin
+++ b/Engine/Plugins/Experimental/DataLink/DataLink.uplugin
@@ -2,7 +2,7 @@
 	"FileVersion": 3,
 	"Version": 1,
 	"VersionName": "1.0",
-	"FriendlyName": "Motion Design Data Link",
+	"FriendlyName": "Motion Design Data",
 	"Description": "",
 	"Category": "Other",
 	"CreatedBy": "",
diff --git a/Engine/Plugins/Experimental/DataLinkOAuth/DataLinkOAuth.uplugin b/Engine/Plugins/Experimental/DataLinkOAuth/DataLinkOAuth.uplugin
--- a/Engine/Plugins/Experimental/DataLinkOAuth/DataLinkOAuth.uplugin
+++ b/Engine/Plugins/Experimental/DataLinkOAuth/DataLinkOAuth.uplugin
@@ -2,8 +2,8 @@
 	"FileVersion": 3,
 	"Version": 1,
 	"VersionName": "1.0",
-	"FriendlyName": "Motion Design Data Link OAuth",
-	"Description": "Data Link functionality for OAuth 2.0",
+	"FriendlyName": "Motion Design Data OAuth",
+	"Description": "Motion Design Data functionality for OAuth 2.0",
 	"Category": "Other",
 	"CreatedBy": "",
 	"CreatedByURL": "",
