Motion Design: renamed Motion Design Data back to Motion Design Data Link

diff --git a/Engine/Plugins/Experimental/AvalancheDataLink/AvalancheDataLink.uplugin b/Engine/Plugins/Experimental/AvalancheDataLink/AvalancheDataLink.uplugin
--- a/Engine/Plugins/Experimental/AvalancheDataLink/AvalancheDataLink.uplugin
+++ b/Engine/Plugins/Experimental/AvalancheDataLink/AvalancheDataLink.uplugin
@@ -2,7 +2,7 @@
 	"FileVersion": 3,
 	"Version": 1,
 	"VersionName": "1.0",
-	"FriendlyName": "Motion Design Data Integration",
+	"FriendlyName": "Motion Design Data Link Integration",
 	"Description": "",
 	"Category": "Other",
 	"CreatedBy": "",
diff --git a/Engine/Plugins/Experimental/DataLink/DataLink.uplugin b/Engine/Plugins/Experimental/DataLink/DataLink.uplugin
--- a/Engine/Plugins/Experimental/DataLink/DataLink.uplugin
+++ b/Engine/Plugins/Experimental/DataLink/DataLink.uplugin
@@ -2,7 +2,7 @@
 	"FileVersion": 3,
 	"Version": 1,
 	"VersionName": "1.0",
-	"FriendlyName": "Motion Design Data",
+	"FriendlyName": "Motion Design Data Link",
 	"Description": "",
 	"Category": "Other",
 	"CreatedBy": "",
diff --git a/Engine/Plugins/Experimental/DataLink/Source/DataLink/Public/DataLinkGraph.h b/Engine/Plugins/Experimental/DataLink/Source/DataLink/Public/DataLinkGraph.h
--- a/Engine/Plugins/Experimental/DataLink/Source/DataLink/Public/DataLinkGraph.h
+++ b/Engine/Plugins/Experimental/DataLink/Source/DataLink/Public/DataLinkGraph.h
@@ -14,7 +14,7 @@ struct FDataLinkPinReference;
 class UEdGraph;
 #endif
 
-UCLASS(MinimalAPI, BlueprintType, DisplayName="Motion Design Data Graph")
+UCLASS(MinimalAPI, BlueprintType, DisplayName="Motion Design Data Link Graph")
 class UDataLinkGraph : public UObject
 {
 	GENERATED_BODY()
diff --git a/Engine/Plugins/Experimental/DataLinkOAuth/DataLinkOAuth.uplugin b/Engine/Plugins/Experimental/DataLinkOAuth/DataLinkOAuth.uplugin
--- a/Engine/Plugins/Experimental/DataLinkOAuth/DataLinkOAuth.uplugin
+++ b/Engine/Plugins/Experimental/DataLinkOAuth/DataLinkOAuth.uplugin
@@ -2,8 +2,8 @@
 	"FileVersion": 3,
 	"Version": 1,
 	"VersionName": "1.0",
-	"FriendlyName": "Motion Design Data OAuth",
-	"Description": "Motion Design Data functionality for OAuth 2.0",
+	"FriendlyName": "Motion Design Data Link OAuth",
+	"Description": "Motion Design Data Link functionality for OAuth 2.0",
 	"Category": "Other",
 	"CreatedBy": "",
 	"CreatedByURL": "",
diff --git a/Engine/Plugins/Experimental/DataLinkOAuth/Source/DataLinkOAuth/Public/DataLinkOAuthSettings.h b/Engine/Plugins/Experimental/DataLinkOAuth/Source/DataLinkOAuth/Public/DataLinkOAuthSettings.h
--- a/Engine/Plugins/Experimental/DataLinkOAuth/Source/DataLinkOAuth/Public/DataLinkOAuthSettings.h
+++ b/Engine/Plugins/Experimental/DataLinkOAuth/Source/DataLinkOAuth/Public/DataLinkOAuthSettings.h
@@ -12,7 +12,7 @@ struct FDataLinkOAuthToken;
 struct FHttpServerRequest;
 struct FInstancedStruct;
 
-UCLASS(MinimalAPI, Abstract, DisplayName="Motion Design Data OAuth Settings")
+UCLASS(MinimalAPI, Abstract, DisplayName="Motion Design Data Link OAuth Settings")
 class UDataLinkOAuthSettings : public UObject
 {
 	GENERATED_BODY()
